Add firefox aurora channel as a device for windows.

#html5

diff --git a/Engine/Config/HTML5/HTML5Engine.ini b/Engine/Config/HTML5/HTML5Engine.ini
--- a/Engine/Config/HTML5/HTML5Engine.ini
+++ b/Engine/Config/HTML5/HTML5Engine.ini
@@ -13,6 +13,7 @@ HeapSize=1024
 [HTML5DevicesWindows]
 Firefox (64 bit)="C:/Program Files/Nightly/firefox.exe"
 Firefox (32 bit)="C:/Program Files (x86)/Mozilla Firefox/firefox.exe"
+Firefox (Aurora)="C:/Program Files (x86)/Aurora/firefox.exe"
 Chrome="C:/Program Files (x86)/Google/Chrome/Application/chrome.exe"
 
 [HTML5DevicesMac]
