Add firefox aurora as another browser device on mac.

diff --git a/Engine/Config/HTML5/HTML5Engine.ini b/Engine/Config/HTML5/HTML5Engine.ini
--- a/Engine/Config/HTML5/HTML5Engine.ini
+++ b/Engine/Config/HTML5/HTML5Engine.ini
@@ -19,6 +19,7 @@ Chrome="C:/Program Files (x86)/Google/Chrome/Application/chrome.exe"
 [HTML5DevicesMac]
 Safari="/Applications/Safari.app/Contents/MacOS/Safari"
 Firefox="/Applications/Firefox.app/Contents/MacOS/firefox-bin"
+Firefox (Aurora)="/Applications/FirefoxAurora.app/Contents/MacOS/firefox-bin"
 
 [HTML5DevicesLinux]
 Firefox="/usr/bin/firefox"
