Better naming for various browser devices.

#html5

diff --git a/Engine/Config/HTML5/HTML5Engine.ini b/Engine/Config/HTML5/HTML5Engine.ini
--- a/Engine/Config/HTML5/HTML5Engine.ini
+++ b/Engine/Config/HTML5/HTML5Engine.ini
@@ -11,11 +11,13 @@ r.setres=800x600
 HeapSize=1024
 
 [HTML5DevicesWindows]
-Firefox (64 bit)="C:/Program Files/Nightly/firefox.exe"
+Firefox Nightly (64 bit)="C:/Program Files/Nightly/firefox.exe"
+Firefox Nightly (32 bit)="C:/Program Files (x86)/Nightly/firefox.exe"
 Firefox (32 bit)="C:/Program Files (x86)/Mozilla Firefox/firefox.exe"
-Firefox (Aurora)="C:/Program Files (x86)/Aurora/firefox.exe"
+Firefox Aurora  (32 bit)="C:/Program Files (x86)/Aurora/firefox.exe"
 Chrome="C:/Program Files (x86)/Google/Chrome/Application/chrome.exe"
 
+;all browsers on mac are 64 bit, 64 bit browsers can allocate more than 512MB of memory to around 1.5 GB. 
 [HTML5DevicesMac]
 Safari="/Applications/Safari.app/Contents/MacOS/Safari"
 Firefox="/Applications/Firefox.app/Contents/MacOS/firefox-bin"
