clean up HTML5Engine.ini  removed commented out stuff. #html5

diff --git a/Engine/Config/HTML5/HTML5Engine.ini b/Engine/Config/HTML5/HTML5Engine.ini
--- a/Engine/Config/HTML5/HTML5Engine.ini
+++ b/Engine/Config/HTML5/HTML5Engine.ini
@@ -9,26 +9,3 @@ benabled=false
 [SystemSettings]
 r.setres=800x600
 
-; These following settings have moved to the User inis and can be set from the editor
-; These old value are left as reference until all old code paths are removed
-;[BuildSettings]
-;HeapSize=1024
-
-;[HTML5DevicesWindows]
-;Firefox Nightly (64 bit)="C:/Program Files/Nightly/firefox.exe"
-;Firefox Nightly (32 bit)="C:/Program Files (x86)/Nightly/firefox.exe"
-;Firefox (32 bit)="C:/Program Files (x86)/Mozilla Firefox/firefox.exe"
-;Firefox Aurora  (32 bit)="C:/Program Files (x86)/Aurora/firefox.exe"
-;Chrome="C:/Program Files (x86)/Google/Chrome/Application/chrome.exe"
-
-;all browsers on mac are 64 bit, 64 bit browsers can allocate more than 512MB of memory to around 1.5 GB. 
-;[HTML5DevicesMac]
-;Safari="/Applications/Safari.app/Contents/MacOS/Safari"
-;Firefox="/Applications/Firefox.app/Contents/MacOS/firefox-bin"
-;Firefox (Aurora)="/Applications/FirefoxAurora.app/Contents/MacOS/firefox-bin"
-
-;[HTML5DevicesLinux]
-;Firefox="/usr/bin/firefox"
-
-;[HTML5SDKPaths]
-;Mac="/path/to/emscripten/sdk"
