Turn on AudioMixerSDL by default for HTML5. #fyi aaron.mcleran #jira none

diff --git a/Engine/Config/HTML5/HTML5Engine.ini b/Engine/Config/HTML5/HTML5Engine.ini
--- a/Engine/Config/HTML5/HTML5Engine.ini
+++ b/Engine/Config/HTML5/HTML5Engine.ini
@@ -1,7 +1,10 @@
 [Audio]
+; This audio backend can be used by running with the -noaudiomixer commandline argument:
 AudioDeviceModuleName=ALAudio
-; Defining below allows switching to audio mixer using -audiomixer commandline
+; This audio backend can be used by running with the -audiomixer commandline argument:
 AudioMixerModuleName=AudioMixerSDL
+; By default, we use the AudioMixerSDL backend.
+UseAudioMixer=true
 
 ; Defines a platform-specific volume headroom (in dB) for audio to provide better platform consistency with respect to volume levels.
 PlatformHeadroomDB=0
