Disabling audio mixer for HTML5

diff --git a/Engine/Config/HTML5/HTML5Engine.ini b/Engine/Config/HTML5/HTML5Engine.ini
--- a/Engine/Config/HTML5/HTML5Engine.ini
+++ b/Engine/Config/HTML5/HTML5Engine.ini
@@ -3,8 +3,6 @@
 AudioDeviceModuleName=ALAudio
 ; This audio backend can be used by running with the -audiomixer commandline argument:
 AudioMixerModuleName=AudioMixerSDL
-; By default, we use the AudioMixerSDL backend.
-UseAudioMixer=true
 
 ; Defines a platform-specific volume headroom (in dB) for audio to provide better platform consistency with respect to volume levels.
 PlatformHeadroomDB=0
